#autosubmit

Renamed UE4Game to UnrealGame.

Horde tested (Editor,Tools & Monolithics). Local ShooterGame + Cook + PIE

diff --git a/Engine/Source/Editor/TurnkeySupport/Private/TurnkeySupportModule.cpp b/Engine/Source/Editor/TurnkeySupport/Private/TurnkeySupportModule.cpp
--- a/Engine/Source/Editor/TurnkeySupport/Private/TurnkeySupportModule.cpp
+++ b/Engine/Source/Editor/TurnkeySupport/Private/TurnkeySupportModule.cpp
@@ -123,7 +123,7 @@ protected:
 				FString TargetName;
 				if (Target == nullptr)
 				{
-					TargetName = TEXT("UE4Game");
+					TargetName = TEXT("UnrealGame");
 				}
 				else
 				{
diff --git a/Engine/Source/Programs/TextureShare/TextureShareSDK/Private/TextureShareDLL.cpp b/Engine/Source/Programs/TextureShare/TextureShareSDK/Private/TextureShareDLL.cpp
--- a/Engine/Source/Programs/TextureShare/TextureShareSDK/Private/TextureShareDLL.cpp
+++ b/Engine/Source/Programs/TextureShare/TextureShareSDK/Private/TextureShareDLL.cpp
@@ -32,7 +32,7 @@ BOOL WINAPI DllMain(HINSTANCE hinstDLL, DWORD fdwReason, LPVOID lpReserved)
 
 #include "Windows/HideWindowsPlatformTypes.h"
 
-//AB If we don't have to implement a module, or application, and we can save 200kb by not depending on "Projects" module, so borrow some code from UE4Game.cpp to bypass it
+//AB If we don't have to implement a module, or application, and we can save 200kb by not depending on "Projects" module, so borrow some code from UnrealGame.cpp to bypass it 
 TCHAR GInternalProjectName[64] = TEXT("");
 
 IMPLEMENT_FOREIGN_ENGINE_DIR()
diff --git a/Engine/Source/Runtime/UE4Game/Private/UE4Game.cpp b/Engine/Source/Runtime/UE4Game/Private/UE4Game.cpp
deleted file mode 100644
diff --git a/Engine/Source/Runtime/UE4Game/UE4Game.Build.cs b/Engine/Source/Runtime/UE4Game/UE4Game.Build.cs
deleted file mode 100644
diff --git a/Engine/Source/Runtime/UnrealGame/Private/UnrealGame.cpp b/Engine/Source/Runtime/UnrealGame/Private/UnrealGame.cpp
new file mode 100644
--- /dev/null
+++ b/Engine/Source/Runtime/UnrealGame/Private/UnrealGame.cpp
@@ -0,0 +1,18 @@
+// Copyright Epic Games, Inc. All Rights Reserved.
+
+#include "CoreTypes.h"
+#include "Modules/Boilerplate/ModuleBoilerplate.h"
+#include "Modules/ModuleManager.h"
+
+// Implement the modules
+IMPLEMENT_MODULE(FDefaultModuleImpl, UnrealGame);
+
+// Variables referenced from Core, and usually declared by IMPLEMENT_PRIMARY_GAME_MODULE. Since UnrealGame is game-agnostic, implement them here.
+#if IS_MONOLITHIC
+PER_MODULE_BOILERPLATE
+bool GIsGameAgnosticExe = true;
+TCHAR GInternalProjectName[64] = TEXT("");
+IMPLEMENT_FOREIGN_ENGINE_DIR()
+IMPLEMENT_LIVE_CODING_ENGINE_DIR()
+IMPLEMENT_LIVE_CODING_PROJECT()
+#endif
diff --git a/Engine/Source/Runtime/UnrealGame/UnrealGame.Build.cs b/Engine/Source/Runtime/UnrealGame/UnrealGame.Build.cs
new file mode 100644
--- /dev/null
+++ b/Engine/Source/Runtime/UnrealGame/UnrealGame.Build.cs
@@ -0,0 +1,31 @@
+// Copyright Epic Games, Inc. All Rights Reserved.
+
+using UnrealBuildTool;
+
+public class UnrealGame : ModuleRules
+{
+	public UnrealGame(ReadOnlyTargetRules Target) : base(Target)
+	{
+		PrivateDependencyModuleNames.Add("Core");
+
+
+        //DynamicallyLoadedModuleNames.Add("OnlineSubsystemNull");
+
+		if (Target.Platform == UnrealTargetPlatform.IOS || Target.Platform == UnrealTargetPlatform.TVOS)
+		{
+			PrivateDependencyModuleNames.AddRange(new string[] { "Core", "CoreUObject", "Engine" });//, "OnlineSubsystem", "OnlineSubsystemUtils" });
+            //DynamicallyLoadedModuleNames.Add("OnlineSubsystemIOS");
+			if (Target.Platform == UnrealTargetPlatform.IOS)
+			{
+                //DynamicallyLoadedModuleNames.Add("OnlineSubsystemFacebook");
+				DynamicallyLoadedModuleNames.Add("IOSAdvertising");
+			}
+		}
+		else if (Target.Platform == UnrealTargetPlatform.Android)
+		{
+			PrivateDependencyModuleNames.AddRange(new string[] { "Core", "CoreUObject", "Engine" });//, "OnlineSubsystem", "OnlineSubsystemUtils" });
+			DynamicallyLoadedModuleNames.Add("AndroidAdvertising");
+            //DynamicallyLoadedModuleNames.Add("OnlineSubsystemGooglePlay");
+		}
+	}
+}
diff --git a/Engine/Source/UnrealClient.Target.cs b/Engine/Source/UnrealClient.Target.cs
--- a/Engine/Source/UnrealClient.Target.cs
+++ b/Engine/Source/UnrealClient.Target.cs
@@ -9,6 +9,6 @@ public class UnrealClientTarget : TargetRules
 	{
 		Type = TargetType.Client;
 		BuildEnvironment = TargetBuildEnvironment.Shared;
-		ExtraModuleNames.Add("UE4Game");
+		ExtraModuleNames.Add("UnrealGame");
 	}
 }
diff --git a/Engine/Source/UnrealEditor.Target.cs b/Engine/Source/UnrealEditor.Target.cs
--- a/Engine/Source/UnrealEditor.Target.cs
+++ b/Engine/Source/UnrealEditor.Target.cs
@@ -10,6 +10,6 @@ public class UnrealEditorTarget : TargetRules
 		Type = TargetType.Editor;
 		BuildEnvironment = TargetBuildEnvironment.Shared;
 		bBuildAllModules = true;
-		ExtraModuleNames.Add("UE4Game");
+		ExtraModuleNames.Add("UnrealGame");
 	}
 }
diff --git a/Engine/Source/UnrealGame.Target.cs b/Engine/Source/UnrealGame.Target.cs
--- a/Engine/Source/UnrealGame.Target.cs
+++ b/Engine/Source/UnrealGame.Target.cs
@@ -11,7 +11,7 @@ public class UnrealGameTarget : TargetRules
 		Type = TargetType.Game;
 		BuildEnvironment = TargetBuildEnvironment.Shared;
 
-		ExtraModuleNames.Add("UE4Game");
+		ExtraModuleNames.Add("UnrealGame");
 
 		if (Target.Platform == UnrealTargetPlatform.IOS)
 		{
diff --git a/Engine/Source/UnrealServer.Target.cs b/Engine/Source/UnrealServer.Target.cs
--- a/Engine/Source/UnrealServer.Target.cs
+++ b/Engine/Source/UnrealServer.Target.cs
@@ -10,6 +10,6 @@ public class UnrealServerTarget : TargetRules
 	{
 		Type = TargetType.Server;
 		BuildEnvironment = TargetBuildEnvironment.Shared;
-		ExtraModuleNames.Add("UE4Game");
+		ExtraModuleNames.Add("UnrealGame");
 	}
 }
