Set UE_GAME_FEATURE_PLUGIN_DETAILS_USE_UTF8_STRING to 1 for UnrealGame

diff --git a/Engine/Source/Runtime/UnrealGame/UnrealGame.Build.cs b/Engine/Source/Runtime/UnrealGame/UnrealGame.Build.cs
--- a/Engine/Source/Runtime/UnrealGame/UnrealGame.Build.cs
+++ b/Engine/Source/Runtime/UnrealGame/UnrealGame.Build.cs
@@ -27,5 +27,7 @@ public class UnrealGame : ModuleRules
 			DynamicallyLoadedModuleNames.Add("AndroidAdvertising");
             //DynamicallyLoadedModuleNames.Add("OnlineSubsystemGooglePlay");
 		}
+
+		PublicDefinitions.Add("UE_GAME_FEATURE_PLUGIN_DETAILS_USE_UTF8_STRING=1");
 	}
 }
