Set UE_GAME_FEATURE_PLUGIN_DETAILS_USE_UTF8_STRING to 1 for UnrealGame.Target.cs instead of UnrealGame.Build.cs

diff --git a/Engine/Source/Runtime/UnrealGame/UnrealGame.Build.cs b/Engine/Source/Runtime/UnrealGame/UnrealGame.Build.cs
--- a/Engine/Source/Runtime/UnrealGame/UnrealGame.Build.cs
+++ b/Engine/Source/Runtime/UnrealGame/UnrealGame.Build.cs
@@ -27,7 +27,5 @@ public class UnrealGame : ModuleRules
 			DynamicallyLoadedModuleNames.Add("AndroidAdvertising");
             //DynamicallyLoadedModuleNames.Add("OnlineSubsystemGooglePlay");
 		}
-
-		PublicDefinitions.Add("UE_GAME_FEATURE_PLUGIN_DETAILS_USE_UTF8_STRING=1");
 	}
 }
diff --git a/Engine/Source/UnrealGame.Target.cs b/Engine/Source/UnrealGame.Target.cs
--- a/Engine/Source/UnrealGame.Target.cs
+++ b/Engine/Source/UnrealGame.Target.cs
@@ -13,5 +13,7 @@ public class UnrealGameTarget : TargetRules
 		BuildEnvironment = TargetBuildEnvironment.Shared;
 
 		ExtraModuleNames.Add("UnrealGame");
+
+		GlobalDefinitions.Add("UE_GAME_FEATURE_PLUGIN_DETAILS_USE_UTF8_STRING=1");
 	}
 }
